Fix RTTI Linux

#Synced-CL 50766215

diff --git a/Engine/Plugins/Runtime/IOSDeviceProfileSelector/Source/IOSPreviewDeviceProfileSelector/IOSPreviewDeviceProfileSelector.Build.cs b/Engine/Plugins/Runtime/IOSDeviceProfileSelector/Source/IOSPreviewDeviceProfileSelector/IOSPreviewDeviceProfileSelector.Build.cs
--- a/Engine/Plugins/Runtime/IOSDeviceProfileSelector/Source/IOSPreviewDeviceProfileSelector/IOSPreviewDeviceProfileSelector.Build.cs
+++ b/Engine/Plugins/Runtime/IOSDeviceProfileSelector/Source/IOSPreviewDeviceProfileSelector/IOSPreviewDeviceProfileSelector.Build.cs
@@ -28,7 +28,8 @@ namespace UnrealBuildTool.Rules
 				}
 				);
 
-			if (IsPlatformAvailable(UnrealTargetPlatform.IOS) && (Target.Type == TargetType.Editor))
+
+			if (IsPlatformAvailable(UnrealTargetPlatform.IOS) && (Target.Type == TargetType.Editor) && (Target.Platform == UnrealTargetPlatform.Win64 || (Target.Platform == UnrealTargetPlatform.Mac)))
 			{
 				PrivateDefinitions.Add("WITH_IOS_DEVICE_DETECTION=1");
 				PublicIncludePathModuleNames.Add("IOSTargetPlatformControls");
