Add Aspect Ratio to IOS Preview Json

#Synced-CL 51089665

diff --git a/Engine/Plugins/Runtime/IOSDeviceProfileSelector/Config/DeviceModelsInfo.ini b/Engine/Plugins/Runtime/IOSDeviceProfileSelector/Config/DeviceModelsInfo.ini
new file mode 100644
--- /dev/null
+++ b/Engine/Plugins/Runtime/IOSDeviceProfileSelector/Config/DeviceModelsInfo.ini
@@ -0,0 +1,13 @@
+;Format Entry=(Devices=<Major_Model_Name>:<Minor_Model_Name1>, <Minor_Model_Name2>, Resx=<ResX_Value>, ResY=<ResY_Value>)
+[DeviceModelToDeviceSpec]
+Entry=(Devices="iPhone10:3,6|iPhone11:2", ResX="2436" ResY="1125")
+Entry=(Devices="iPhone11:4,6|iPhone12:5", ResX="2688", ResY="1242")
+Entry=(Devices="iPhone11:8|iPhone12:1", ResX="1792", ResY="828")
+Entry=(Devices="iPhone13:1|iPhone14:4", ResX="2340", ResY="1080")
+Entry=(Devices="iPhone13:2,3|iPhone14:2,5,7", ResX="2532", ResY="1170")
+Entry=(Devices="iPhone13:4|iPhone14:3,8", ResX="2778", ResY="1284")
+Entry=(Devices="iPhone15:2,4|iPhone16:1", ResX="2556", ResY="1179")
+Entry=(Devices="iPhone15:3,5|iPhone16:2|iPhone17:4", ResX="2796", ResY="1290")
+Entry=(Devices="iPhone17:1,2|iPhone18:1,2", ResX="2622", ResY="1206")
+Entry=(Devices="iPhone17:3|iPhone18:3", ResX="2868", ResY="1320")
+Entry=(Devices="iPhone18:4", ResX="2736", ResY="1260")
\ No newline at end of file
diff --git a/Engine/Plugins/Runtime/IOSDeviceProfileSelector/Source/IOSPreviewDeviceProfileSelector/IOSPreviewDeviceProfileSelector.Build.cs b/Engine/Plugins/Runtime/IOSDeviceProfileSelector/Source/IOSPreviewDeviceProfileSelector/IOSPreviewDeviceProfileSelector.Build.cs
--- a/Engine/Plugins/Runtime/IOSDeviceProfileSelector/Source/IOSPreviewDeviceProfileSelector/IOSPreviewDeviceProfileSelector.Build.cs
+++ b/Engine/Plugins/Runtime/IOSDeviceProfileSelector/Source/IOSPreviewDeviceProfileSelector/IOSPreviewDeviceProfileSelector.Build.cs
@@ -25,6 +25,7 @@ namespace UnrealBuildTool.Rules
 					"Json",
 					"JsonUtilities",
 					"PIEPreviewDeviceSpecification",
+					"Projects",
 				}
 				);
 
diff --git a/Engine/Plugins/Runtime/IOSDeviceProfileSelector/Source/IOSPreviewDeviceProfileSelector/Private/IOSPreviewDeviceProfileSelectorModule.cpp b/Engine/Plugins/Runtime/IOSDeviceProfileSelector/Source/IOSPreviewDeviceProfileSelector/Private/IOSPreviewDeviceProfileSelectorModule.cpp
--- a/Engine/Plugins/Runtime/IOSDeviceProfileSelector/Source/IOSPreviewDeviceProfileSelector/Private/IOSPreviewDeviceProfileSelectorModule.cpp
+++ b/Engine/Plugins/Runtime/IOSDeviceProfileSelector/Source/IOSPreviewDeviceProfileSelector/Private/IOSPreviewDeviceProfileSelectorModule.cpp
@@ -10,6 +10,7 @@
 #include "Internationalization/Regex.h"
 #include "Misc/ConfigCacheIni.h"
 #include "DataDrivenShaderPlatformInfo.h"
+#include "Interfaces/IPluginManager.h"
 #if WITH_IOS_DEVICE_DETECTION
 #include "IOSDeviceHelper.h"
 #endif
@@ -19,9 +20,44 @@ DEFINE_LOG_CATEGORY_STATIC(LogIOSDPSelector, Log, All)
 
 IMPLEMENT_MODULE(FIOSPreviewDeviceProfileSelectorModule, IOSPreviewDeviceProfileSelector);
 
+struct FIOSPreviewDeviceSpec
+{
+	float ResX;
+	float ResY;
+};
+
+static TMap<TArray<FString>, FIOSPreviewDeviceSpec> DeviceModelToDeviceSpec;
 
 void FIOSPreviewDeviceProfileSelectorModule::StartupModule()
 {
+	FString IniPath = FPaths::Combine(IPluginManager::Get().FindPlugin(TEXT("IOSDeviceProfileSelector"))->GetBaseDir(), TEXT("Config/DeviceModelsInfo.ini"));
+	TArray<FString> RawInputs;
+	FConfigFile DeviceModelsInfo;
+	DeviceModelsInfo.Read(IniPath);
+	if (DeviceModelsInfo.GetArray(TEXT("DeviceModelToDeviceSpec"), TEXT("Entry"), RawInputs))
+	{
+		for (const FString& Line : RawInputs)
+		{
+			FString RawDeviceModels, RawResXValues, RawResYValues;
+			if (FParse::Value(*Line, TEXT("Devices="), RawDeviceModels) &&
+				FParse::Value(*Line, TEXT("ResX="), RawResXValues) &&
+				FParse::Value(*Line, TEXT("ResY="), RawResYValues))
+			{
+				RawDeviceModels = RawDeviceModels.TrimQuotes().TrimStartAndEnd();
+				RawResXValues = RawResXValues.TrimQuotes().TrimStartAndEnd();
+				RawResYValues = RawResYValues.TrimQuotes().TrimStartAndEnd();
+
+				TArray<FString> DeviceModels;
+				RawDeviceModels.ParseIntoArray(DeviceModels, TEXT("|"), true);
+
+				FIOSPreviewDeviceSpec DeviceSpec;
+				DeviceSpec.ResX = FCString::Atof(*RawResXValues);
+				DeviceSpec.ResY = FCString::Atof(*RawResYValues);
+
+				DeviceModelToDeviceSpec.Add(DeviceModels, DeviceSpec);
+			}
+		}
+	}
 }
 
 
@@ -39,7 +75,7 @@ FString const FIOSPreviewDeviceProfileSelectorModule::GetRuntimeDeviceProfileNam
 	return FString();
 }
 
-const FString GetDeviceProfileNameFromDeviceModel(FString& DeviceModel)
+const FString GetDeviceProfileNameFromDeviceModel(const FString& DeviceModel)
 {
 	TArray<FString> MappingSections;
 	MappingSections.Add(TEXT("IOSDeviceMappings"));
@@ -245,6 +281,47 @@ const EShaderPlatform FIOSPreviewDeviceProfileSelectorModule::GetPreviewShaderPl
 
 void FIOSPreviewDeviceProfileSelectorModule::SetSelectorProperties(const TMap<FString, FString>& InSelectorProperties)
 {
+	const FString IPhone = TEXT("iPhone");
+	const FString IPad = TEXT("iPad");
+	if (const FString* DeviceModel = InSelectorProperties.Find(TEXT("SRC_DeviceModel")))
+	{
+		// Fallbacks if we do not find the model.
+		if (DeviceModel->StartsWith(IPhone))
+		{
+			ConstrainedAspectRatio = 19.5f/9.f;
+		}
+		else if (DeviceModel->StartsWith(IPad))
+		{
+			ConstrainedAspectRatio = 4.f/3.f;
+		}
+
+		FIOSPreviewDeviceSpec DeviceSpec;
+		for (const auto& Pair : DeviceModelToDeviceSpec)
+		{
+			for (const FString& DeviceModelName : Pair.Key)
+			{
+				FString MajorModelName;
+				FString MinorModelNumberArray;
+				DeviceModelName.Split(":", &MajorModelName, &MinorModelNumberArray);
+
+				TArray<FString> MinorModelNumbers;
+				MinorModelNumberArray.ParseIntoArray(MinorModelNumbers, TEXT(","), true);
+				for (const FString& MinorModelNumber : MinorModelNumbers)
+				{
+					FString FullModelNumberToVerify = MajorModelName + TEXT(",") + MinorModelNumber;
+					if (*DeviceModel == FullModelNumberToVerify)
+					{
+						DeviceSpec = Pair.Value;
+						if(DeviceSpec.ResY != 0.f)
+						{
+							ConstrainedAspectRatio = DeviceSpec.ResX / DeviceSpec.ResY;
+						}
+						break;
+					}
+				}
+			}
+		}
+	}
 	SelectorProperties = InSelectorProperties;
 }
 
